#ReplayServer: Add batch files for builder

diff --git a/Engine/Source/Programs/ReplayServer/build.bat b/Engine/Source/Programs/ReplayServer/build.bat
new file mode 100644
--- /dev/null
+++ b/Engine/Source/Programs/ReplayServer/build.bat
@@ -0,0 +1,31 @@
+@echo off
+echo Building replay server...
+
+if exist Intermediate rd Intermediate /s/q
+if exist Binaries rd Binaries /s/q
+
+md Intermediate
+md Binaries
+
+if not defined JAVA_PATH (
+	set JAVA_PATH="C:\Program Files\Java\jdk1.8.0_31\bin"
+)
+
+if %JAVA_PATH% == "" (
+	set JAVA_PATH="C:\Program Files\Java\jdk1.8.0_31\bin"
+)
+
+%JAVA_PATH%\javac -d Intermediate -sourcepath src\com\epicgames\replayserver\ -cp ThirdParty\Jetty\*;ThirdParty\GSon\*;ThirdParty\MongoDB\* src\com\epicgames\replayserver\*.java
+%JAVA_PATH%\jar cfm Binaries\ReplayServer.jar manifest.txt -C Intermediate com\epicgames\replayserver
+
+xcopy /q ThirdParty\Jetty\jetty-util-9.2.7.v20150116.jar Binaries\ThirdParty\Jetty\
+xcopy /q ThirdParty\Jetty\jetty-server-9.2.7.v20150116.jar Binaries\ThirdParty\Jetty\
+xcopy /q ThirdParty\Jetty\servlet-api-3.1.jar Binaries\ThirdParty\Jetty\
+xcopy /q ThirdParty\Jetty\jetty-security-9.2.7.v20150116.jar Binaries\ThirdParty\Jetty\
+xcopy /q ThirdParty\Jetty\jetty-http-9.2.7.v20150116.jar Binaries\ThirdParty\Jetty\
+xcopy /q ThirdParty\Jetty\jetty-io-9.2.7.v20150116.jar Binaries\ThirdParty\Jetty\
+xcopy /q ThirdParty\Jetty\jetty-servlet-9.2.7.v20150116.jar Binaries\ThirdParty\Jetty\
+xcopy /q ThirdParty\MongoDB\mongo-java-driver-2.13.0.jar Binaries\ThirdParty\MongoDB\
+xcopy /q ThirdParty\Gson\gson-2.3.1.jar Binaries\ThirdParty\Gson\
+
+if exist Intermediate rd Intermediate /s/q
diff --git a/Engine/Source/Programs/ReplayServer/manifest.txt b/Engine/Source/Programs/ReplayServer/manifest.txt
new file mode 100644
--- /dev/null
+++ b/Engine/Source/Programs/ReplayServer/manifest.txt
@@ -0,0 +1,11 @@
+Main-Class: com.epicgames.replayserver.ReplayServer
+Class-Path: 
+  ThirdParty\Jetty\jetty-util-9.2.7.v20150116.jar
+  ThirdParty\Jetty\jetty-server-9.2.7.v20150116.jar
+  ThirdParty\Jetty\servlet-api-3.1.jar
+  ThirdParty\Jetty\jetty-security-9.2.7.v20150116.jar
+  ThirdParty\Jetty\jetty-http-9.2.7.v20150116.jar
+  ThirdParty\Jetty\jetty-io-9.2.7.v20150116.jar
+  ThirdParty\Jetty\jetty-servlet-9.2.7.v20150116.jar
+  ThirdParty\MongoDB\mongo-java-driver-2.13.0.jar
+  ThirdParty\Gson\gson-2.3.1.jar
