#ReplayServer: Change env variable name

diff --git a/Engine/Source/Programs/ReplayServer/build.bat b/Engine/Source/Programs/ReplayServer/build.bat
--- a/Engine/Source/Programs/ReplayServer/build.bat
+++ b/Engine/Source/Programs/ReplayServer/build.bat
@@ -7,16 +7,16 @@ if exist Binaries rd Binaries /s/q
 md Intermediate
 md Binaries
 
-if not defined JAVA_PATH (
-	set JAVA_PATH="C:\Program Files\Java\jdk1.8.0_31\bin"
+if not defined UE_JAVA_SDK_PATH (
+	set UE_JAVA_SDK_PATH="C:\Program Files\Java\jdk1.8.0_31\bin"
 )
 
-if %JAVA_PATH% == "" (
-	set JAVA_PATH="C:\Program Files\Java\jdk1.8.0_31\bin"
+if %UE_JAVA_SDK_PATH% == "" (
+	set UE_JAVA_SDK_PATH="C:\Program Files\Java\jdk1.8.0_31\bin"
 )
 
-%JAVA_PATH%\javac -d Intermediate -sourcepath src\com\epicgames\replayserver\ -cp ThirdParty\Jetty\*;ThirdParty\GSon\*;ThirdParty\MongoDB\* src\com\epicgames\replayserver\*.java
-%JAVA_PATH%\jar cfm Binaries\ReplayServer.jar manifest.txt -C Intermediate com\epicgames\replayserver
+%UE_JAVA_SDK_PATH%\javac -d Intermediate -sourcepath src\com\epicgames\replayserver\ -cp ThirdParty\Jetty\*;ThirdParty\GSon\*;ThirdParty\MongoDB\* src\com\epicgames\replayserver\*.java
+%UE_JAVA_SDK_PATH%\jar cfm Binaries\ReplayServer.jar manifest.txt -C Intermediate com\epicgames\replayserver
 
 xcopy /q ThirdParty\Jetty\jetty-util-9.2.7.v20150116.jar Binaries\ThirdParty\Jetty\
 xcopy /q ThirdParty\Jetty\jetty-server-9.2.7.v20150116.jar Binaries\ThirdParty\Jetty\
