#replayserver: Build script fixes

diff --git a/Engine/Source/Programs/ReplayServer/build.bat b/Engine/Source/Programs/ReplayServer/build.bat
--- a/Engine/Source/Programs/ReplayServer/build.bat
+++ b/Engine/Source/Programs/ReplayServer/build.bat
@@ -1,20 +1,34 @@
 @echo off
 echo Building replay server...
 
-if exist Intermediate rd Intermediate /s/q
-if exist Binaries rd Binaries /s/q
-
-md Intermediate
-md Binaries
-
 if not defined UE_JAVA_SDK_PATH (
-	set UE_JAVA_SDK_PATH="C:\Program Files\Java\jdk1.8.0_31\bin"
+	echo UE_JAVA_SDK_PATH not defined
+	goto Exit
 )
 
 if %UE_JAVA_SDK_PATH% == "" (
-	set UE_JAVA_SDK_PATH="C:\Program Files\Java\jdk1.8.0_31\bin"
+	echo UE_JAVA_SDK_PATH is blank
+	goto Exit
+)
+
+echo Using Java SDK path: %UE_JAVA_SDK_PATH%
+
+if not exist %UE_JAVA_SDK_PATH%\javac.exe (
+	echo Java SDK appears missing
+	goto Exit
 )
 
+if not exist %UE_JAVA_SDK_PATH%\jar.exe (
+	echo Java SDK appears missing
+	goto Exit
+)
+
+if exist Intermediate rd Intermediate /s/q
+if exist Binaries rd Binaries /s/q
+
+md Intermediate
+md Binaries
+
 %UE_JAVA_SDK_PATH%\javac -d Intermediate -sourcepath src\com\epicgames\replayserver\ -cp ThirdParty\Jetty\*;ThirdParty\GSon\*;ThirdParty\MongoDB\* src\com\epicgames\replayserver\*.java
 %UE_JAVA_SDK_PATH%\jar cfm Binaries\ReplayServer.jar manifest.txt -C Intermediate com\epicgames\replayserver
 
@@ -29,3 +43,5 @@ xcopy /q ThirdParty\MongoDB\mongo-java-driver-2.13.0.jar Binaries\ThirdParty\Mon
 xcopy /q ThirdParty\Gson\gson-2.3.1.jar Binaries\ThirdParty\Gson\
 
 if exist Intermediate rd Intermediate /s/q
+
+:Exit
