#replayserver: Add ability to load properties from config

diff --git a/Engine/Source/Programs/ReplayServer/ReplayServer.Properties b/Engine/Source/Programs/ReplayServer/ReplayServer.Properties
new file mode 100644
--- /dev/null
+++ b/Engine/Source/Programs/ReplayServer/ReplayServer.Properties
@@ -0,0 +1,5 @@
+#ReplayServer Properties
+#Mon May 18 16:28:16 EDT 2015
+quickTaskFrequencyInSeconds=10
+longTaskFrequencyInHours=1
+httpPort=80
\ No newline at end of file
diff --git a/Engine/Source/Programs/ReplayServer/src/com/epicgames/replayserver/ReplayProps.java b/Engine/Source/Programs/ReplayServer/src/com/epicgames/replayserver/ReplayProps.java
new file mode 100644
--- /dev/null
+++ b/Engine/Source/Programs/ReplayServer/src/com/epicgames/replayserver/ReplayProps.java
@@ -0,0 +1,46 @@
+package com.epicgames.replayserver;
+
+import java.io.File;
+import java.io.FileInputStream;
+import java.io.FileOutputStream;
+import java.io.IOException;
+import java.io.InputStream;
+import java.util.Properties;
+
+public class ReplayProps 
+{
+	static Properties props = new Properties();
+	
+	ReplayProps()
+	{
+	}
+	
+	static void Init() throws IOException
+	{
+		try ( final InputStream is = new FileInputStream( new File( "ReplayServer.Properties" ) ) )
+		{
+			props.load( is );
+		}
+	}
+	
+	static String getString( final String key, final String defaultValue )
+	{
+		return props.getProperty( key, defaultValue );
+	}
+	
+	static int getInt( final String key, final String defaultValue )
+	{
+		return Integer.parseInt( getString( key, defaultValue ) );
+	}
+	
+	static void SaveProps()
+	{
+		try ( final FileOutputStream out = new FileOutputStream( new File( "ReplayServer.Properties" ) ) ) 
+		{
+			props.store( out, "ReplayServer Properties" );
+		}
+		catch ( Exception e )
+		{
+		}		
+	}
+}
diff --git a/Engine/Source/Programs/ReplayServer/src/com/epicgames/replayserver/ReplayServer.java b/Engine/Source/Programs/ReplayServer/src/com/epicgames/replayserver/ReplayServer.java
--- a/Engine/Source/Programs/ReplayServer/src/com/epicgames/replayserver/ReplayServer.java
+++ b/Engine/Source/Programs/ReplayServer/src/com/epicgames/replayserver/ReplayServer.java
@@ -28,17 +28,32 @@ public class ReplayServer
 	
 	public static void main( String[] args ) throws Exception
     {
+		ReplayProps.Init();
+
 		ReplayDB.init();
 
 		ReplayLogger.log( Level.INFO, "Log opened..." );
 
 		org.eclipse.jetty.util.log.Log.setLog( ReplayLogger.replayLogger );
+
+		final int quickTaskFrequencyInSeconds 	= ReplayProps.getInt( "quickTaskFrequencyInSeconds", "10" );
+		final int longTaskFrequencyInHours 		= ReplayProps.getInt( "longTaskFrequencyInHours", "1" );
 	
 		timer = new Timer();
-	    timer.schedule( new CleanupTask.QuickTask(), 5 * 1000, 10 * 1000 );		
-	    timer.schedule( new CleanupTask.LongTask(), 10 * 1000, 60 * 60 * 1000 * 1 );	// Run once after 10 seconds, then once every hour		
+
+		if ( quickTaskFrequencyInSeconds != 0 )
+		{
+			timer.schedule( new CleanupTask.QuickTask(), quickTaskFrequencyInSeconds * 1000, quickTaskFrequencyInSeconds * 1000 );		
+		}
+
+		if ( quickTaskFrequencyInSeconds != 0 )
+		{
+			timer.schedule( new CleanupTask.LongTask(), 60 * 60 * 1000 * longTaskFrequencyInHours, 60 * 60 * 1000 * longTaskFrequencyInHours );		
+		}
 		
-		final Server server = new Server( 80 );
+		final int httpPort = ReplayProps.getInt( "httpPort", "80" );
+
+		final Server server = new Server( httpPort );
  
         ServletContextHandler context = new ServletContextHandler( ServletContextHandler.SESSIONS | ServletContextHandler.SECURITY );
         context.setContextPath( "/" );
@@ -48,7 +63,7 @@ public class ReplayServer
         /*
         ServerConnector http = new ServerConnector( server );
         //http.setHost( "localhost" );
-        http.setPort( 80 );
+        http.setPort( httpPort );
         http.setIdleTimeout( 45000 );
 
         server.addConnector( http );
