#replayserver: Copy properties file to binaries output

diff --git a/Engine/Source/Programs/ReplayServer/build.bat b/Engine/Source/Programs/ReplayServer/build.bat
--- a/Engine/Source/Programs/ReplayServer/build.bat
+++ b/Engine/Source/Programs/ReplayServer/build.bat
@@ -42,6 +42,8 @@ xcopy /q ThirdParty\Jetty\jetty-servlet-9.2.7.v20150116.jar Binaries\ThirdParty\
 xcopy /q ThirdParty\MongoDB\mongo-java-driver-2.13.0.jar Binaries\ThirdParty\MongoDB\
 xcopy /q ThirdParty\Gson\gson-2.3.1.jar Binaries\ThirdParty\Gson\
 
+xcopy /q ReplayServer.Properties Binaries\
+
 if exist Intermediate rd Intermediate /s/q
 
 :Exit
