#replayserver: Add more config options

diff --git a/Engine/Source/Programs/ReplayServer/ReplayServer.Properties b/Engine/Source/Programs/ReplayServer/ReplayServer.Properties
--- a/Engine/Source/Programs/ReplayServer/ReplayServer.Properties
+++ b/Engine/Source/Programs/ReplayServer/ReplayServer.Properties
@@ -1,5 +1,22 @@
-#ReplayServer Properties
-#Mon May 18 16:28:16 EDT 2015
+#ReplayServer properties
+
+#General properties
+
 quickTaskFrequencyInSeconds=10
 longTaskFrequencyInHours=1
-httpPort=80
\ No newline at end of file
+
+httpPort=80
+
+enableBuiltInWebServer=0
+
+#MongoDB properties
+
+mongoDB_Host=localhost
+mongoDB_Port=27017
+
+mongoDB_DBName=replayDB
+mongoDB_GridFSBucket=replayFS
+mongoDB_ReplayCollection=replayCollection
+mongoDB_ViewerCollection=viewerCollection
+mongoDB_EventCollection=eventCollection
+mongoDB_LogCollection=logCollection
diff --git a/Engine/Source/Programs/ReplayServer/src/com/epicgames/replayserver/MongoDB.java b/Engine/Source/Programs/ReplayServer/src/com/epicgames/replayserver/MongoDB.java
--- a/Engine/Source/Programs/ReplayServer/src/com/epicgames/replayserver/MongoDB.java
+++ b/Engine/Source/Programs/ReplayServer/src/com/epicgames/replayserver/MongoDB.java
@@ -41,22 +41,22 @@ public class MongoDB implements BaseDB
 
 	public MongoDB() throws UnknownHostException
 	{
-		mongoClient = new MongoClient( "localhost" , 27017 );
+		mongoClient = new MongoClient( ReplayProps.getString( "mongoDB_Host", "localhost" ) , ReplayProps.getInt( "mongoDB_Port", "27017" ) );
 
 		final boolean bResetDB 		= false;
 		final boolean bResetIndexes = false;
 
 		if ( bResetDB )
 		{
-			mongoClient.dropDatabase( "replayDB" );
+			mongoClient.dropDatabase( ReplayProps.getString( "mongoDB_DBName", "replayDB" ) );
 		}
 
-		db 			= mongoClient.getDB( "replayDB" );
-		grid 		= new GridFS( db, "replayFS" );
-		replayColl 	= db.getCollection( "replayCollection" );
-		viewerColl 	= db.getCollection( "viewerCollection" );
-		eventColl 	= db.getCollection( "eventCollection" );
-		logColl 	= db.getCollection( "logCollection" );
+		db 			= mongoClient.getDB( ReplayProps.getString( "mongoDB_DBName", "replayDB" ) );
+		grid 		= new GridFS( db, ReplayProps.getString( "mongoDB_GridFSBucket", "replayFS" ) );
+		replayColl 	= db.getCollection( ReplayProps.getString( "mongoDB_ReplayCollection", "replayCollection" ) );
+		viewerColl 	= db.getCollection( ReplayProps.getString( "mongoDB_ViewerCollection", "viewerCollection" ) );
+		eventColl 	= db.getCollection( ReplayProps.getString( "mongoDB_EventCollection", "eventCollection" ) );
+		logColl 	= db.getCollection( ReplayProps.getString( "mongoDB_LogCollection", "logCollection" ) );
 				
 		if ( bResetIndexes )
 		{
diff --git a/Engine/Source/Programs/ReplayServer/src/com/epicgames/replayserver/ReplayServer.java b/Engine/Source/Programs/ReplayServer/src/com/epicgames/replayserver/ReplayServer.java
--- a/Engine/Source/Programs/ReplayServer/src/com/epicgames/replayserver/ReplayServer.java
+++ b/Engine/Source/Programs/ReplayServer/src/com/epicgames/replayserver/ReplayServer.java
@@ -114,7 +114,11 @@ public class ReplayServer
         context.addServlet( new ServletHolder( new EnumEvents() ), "/enumevents/*" );
         context.addServlet( new ServletHolder( new RefreshViewer() ), "/refreshviewer/*" );
         context.addServlet( new ServletHolder( new ViewFile() ), "/viewfile/*" );
-        context.addServlet( new ServletHolder( new Index() ), "/*" );
+        
+        if ( ReplayProps.getInt( "enableBuiltInWebServer", "1" ) == 1 )
+        {
+            context.addServlet( new ServletHolder( new Index() ), "/*" );
+        }
         
         // Password protect the delete session page
         context.addServlet(new ServletHolder(new HttpServlet() 
