#replayserver: Enable web server by default

diff --git a/Engine/Source/Programs/ReplayServer/ReplayServer.Properties b/Engine/Source/Programs/ReplayServer/ReplayServer.Properties
--- a/Engine/Source/Programs/ReplayServer/ReplayServer.Properties
+++ b/Engine/Source/Programs/ReplayServer/ReplayServer.Properties
@@ -7,7 +7,7 @@ longTaskFrequencyInHours=1
 
 httpPort=80
 
-enableBuiltInWebServer=0
+enableBuiltInWebServer=1
 
 #MongoDB properties
 
