#replayserver: More config properties

diff --git a/Engine/Source/Programs/ReplayServer/ReplayServer.Properties b/Engine/Source/Programs/ReplayServer/ReplayServer.Properties
--- a/Engine/Source/Programs/ReplayServer/ReplayServer.Properties
+++ b/Engine/Source/Programs/ReplayServer/ReplayServer.Properties
@@ -20,3 +20,9 @@ mongoDB_ReplayCollection=replayCollection
 mongoDB_ViewerCollection=viewerCollection
 mongoDB_EventCollection=eventCollection
 mongoDB_LogCollection=logCollection
+
+mongoDB_cleanupStaleLiveSessions_MaxMinutes=1
+mongoDB_cleanupOldSessions_MaxDays=30
+mongoDB_cleanupStaleViewers_MaxMinutes=1
+mongoDB_cleanupLogs_MaxDays=14
+mongoDB_cleanupLogs_MaxFinestDays=3
\ No newline at end of file
diff --git a/Engine/Source/Programs/ReplayServer/src/com/epicgames/replayserver/MongoDB.java b/Engine/Source/Programs/ReplayServer/src/com/epicgames/replayserver/MongoDB.java
--- a/Engine/Source/Programs/ReplayServer/src/com/epicgames/replayserver/MongoDB.java
+++ b/Engine/Source/Programs/ReplayServer/src/com/epicgames/replayserver/MongoDB.java
@@ -829,7 +829,7 @@ public class MongoDB implements BaseDB
 	void cleanupStaleLiveSessions()
 	{
  		// Find all sessions that haven't been refreshed for more than 60 seconds
-		final long maxMinutes = 1;
+		final long maxMinutes = ReplayProps.getInt( "mongoDB_cleanupStaleLiveSessions_MaxMinutes", "1" );
 
 		final DBCursor cursor = replayColl.find( new BasicDBObject( "modified", new BasicDBObject( "$lt", System.currentTimeMillis() - 60 * 1000 * maxMinutes ) ) );//.limit( 20 );
 
@@ -839,17 +839,17 @@ public class MongoDB implements BaseDB
 	void cleanupOldSessions()
 	{
  		// Delete all sessions older than 30 days
-		final long maxDays = 30;
+		final long maxDays = ReplayProps.getInt( "mongoDB_cleanupOldSessions_MaxDays", "30" );
 
 		final DBCursor cursor = replayColl.find( new BasicDBObject( "created", new BasicDBObject( "$lt", System.currentTimeMillis() - 60 * 60 * 24 * maxDays * 1000 ) ) );//.limit( 20 );
 
 		cleanupSessions( cursor, false );
 	}
 
-	void cleanupOldViewers()
+	void cleanupStaleViewers()
 	{
  		// Find all viewers that haven't been refreshed for more than 60 seconds
-		final long maxMinutes = 1;
+		final long maxMinutes = ReplayProps.getInt( "mongoDB_cleanupStaleViewers_MaxMinutes", "1" );
 
 		final DBCursor cursor = viewerColl.find( new BasicDBObject( "modified", new BasicDBObject( "$lt", System.currentTimeMillis() - 60 * 1000 * maxMinutes ) ) );//.limit( 20 );
 
@@ -875,17 +875,17 @@ public class MongoDB implements BaseDB
 	
 	void cleanupLogs()
 	{
-		final long maxDays = 14;
+		final long maxDays = ReplayProps.getInt( "mongoDB_cleanupLogs_MaxDays", "14" );
 		logColl.remove( new BasicDBObject( "created", new BasicDBObject( "$lt", System.currentTimeMillis() - 60 * 60 * 24 * maxDays * 1000 ) ) );
 
-		final long maxFinestDays = 3;
+		final long maxFinestDays = ReplayProps.getInt( "mongoDB_cleanupLogs_MaxFinestDays", "3" );
 		logColl.remove( new BasicDBObject( "level", new BasicDBObject( "$eq", Level.FINEST.intValue() ) ).append( "created", new BasicDBObject( "$lt", System.currentTimeMillis() - 60 * 60 * 24 * maxFinestDays * 1000 ) ) );
 	}
 	
 	public void quickCleanup()
 	{
 		cleanupStaleLiveSessions();
-		cleanupOldViewers();
+		cleanupStaleViewers();
 	}
 
 	public void longCleanup()
