#replayserver: Add ability to track recently viewed replays

diff --git a/Engine/Source/Programs/ReplayServer/ReplayServer.Properties b/Engine/Source/Programs/ReplayServer/ReplayServer.Properties
--- a/Engine/Source/Programs/ReplayServer/ReplayServer.Properties
+++ b/Engine/Source/Programs/ReplayServer/ReplayServer.Properties
@@ -19,6 +19,7 @@ mongoDB_GridFSBucket=replayFS
 mongoDB_ReplayCollection=replayCollection
 mongoDB_ViewerCollection=viewerCollection
 mongoDB_EventCollection=eventCollection
+mongoDB_RecentCollection=recentCollection
 mongoDB_LogCollection=logCollection
 
 mongoDB_cleanupStaleLiveSessions_MaxMinutes=1
diff --git a/Engine/Source/Programs/ReplayServer/src/com/epicgames/replayserver/BaseDB.java b/Engine/Source/Programs/ReplayServer/src/com/epicgames/replayserver/BaseDB.java
--- a/Engine/Source/Programs/ReplayServer/src/com/epicgames/replayserver/BaseDB.java
+++ b/Engine/Source/Programs/ReplayServer/src/com/epicgames/replayserver/BaseDB.java
@@ -73,12 +73,13 @@ public interface BaseDB
 	public byte[] readEventData( final String eventId ) throws ReplayException;
 	public List<BaseDB.EventInfo> enumerateEvents( final String sessionName, final String eventGroup ) throws ReplayException;
 
-	public void createViewer( final String sessionName, final String viewerName ) throws ReplayException;
+	public void createViewer( final String sessionName, final String viewerName, final String userName ) throws ReplayException;
 	public void refreshViewer( final String sessionName, final String viewerName ) throws ReplayException;
 	public void deleteViewer( final String sessionName, final String viewerName ) throws ReplayException;
 	public ReplaySessionInfo getSessionInfo( final String sessionName ) throws ReplayException;
 	public int getNumSessions( final String appName, final int version, final int changelist );
 	public List<ReplaySessionInfo> discoverSessions( final String appName, final int version, final int changelist, final int limit );
+	public List<ReplaySessionInfo> getRecentSessions( final String appName, final int version, final int changelist, final String userName, final int limit );
 
 	public void log( final Level level, final String str );
 	public int printLog( final PrintWriter writer, final Level level, final String subFilter );
diff --git a/Engine/Source/Programs/ReplayServer/src/com/epicgames/replayserver/EnumerateSessions.java b/Engine/Source/Programs/ReplayServer/src/com/epicgames/replayserver/EnumerateSessions.java
--- a/Engine/Source/Programs/ReplayServer/src/com/epicgames/replayserver/EnumerateSessions.java
+++ b/Engine/Source/Programs/ReplayServer/src/com/epicgames/replayserver/EnumerateSessions.java
@@ -44,6 +44,7 @@ public class EnumerateSessions extends HttpServlet
     	final String appName 		= request.getParameter( "App" );
     	final String versionString 	= request.getParameter( "Version" );
     	final String clString 		= request.getParameter( "CL" );
+    	final String userName 		= request.getParameter( "User" );
 
     	if ( appName == null )
     	{
@@ -59,7 +60,16 @@ public class EnumerateSessions extends HttpServlet
 
 		final int hardLimit = 500;
 
-    	List< ReplaySessionInfo > replaySessions = ReplayDB.discoverSessions( appName, version, changelist, hardLimit );
+    	List< ReplaySessionInfo > replaySessions = null;
+    	
+    	if ( userName != null )
+    	{
+    		replaySessions = ReplayDB.getRecentSessions( appName, version, changelist, userName, hardLimit );
+    	}
+    	else
+    	{
+    		replaySessions = ReplayDB.discoverSessions( appName, version, changelist, hardLimit );
+    	}
 
 		Gson gson = new GsonBuilder().setDateFormat( "yyyy-MM-dd'T'HH:mm:ss" ).create();
 
diff --git a/Engine/Source/Programs/ReplayServer/src/com/epicgames/replayserver/FileDB.java b/Engine/Source/Programs/ReplayServer/src/com/epicgames/replayserver/FileDB.java
--- a/Engine/Source/Programs/ReplayServer/src/com/epicgames/replayserver/FileDB.java
+++ b/Engine/Source/Programs/ReplayServer/src/com/epicgames/replayserver/FileDB.java
@@ -251,7 +251,7 @@ public class FileDB implements BaseDB
 		return FileUtils.fileExists( getSessionViewerFilename( sessionName, viewerName ) );
 	}
 
-	public void createViewer( final String sessionName, final String viewerName ) throws ReplayException
+	public void createViewer( final String sessionName, final String viewerName, final String userName ) throws ReplayException
 	{
 		if ( ViewerExists( sessionName, viewerName ) )
 		{
@@ -418,6 +418,11 @@ public class FileDB implements BaseDB
         return sessions;
 	}
 
+	public List<ReplaySessionInfo> getRecentSessions( final String appName, final int version, final int changelist, final String userName, final int limit )
+	{
+		return new ArrayList<ReplaySessionInfo>();	
+	}
+
 	public void log( final Level level, final String str )
 	{
 	}
diff --git a/Engine/Source/Programs/ReplayServer/src/com/epicgames/replayserver/MongoDB.java b/Engine/Source/Programs/ReplayServer/src/com/epicgames/replayserver/MongoDB.java
--- a/Engine/Source/Programs/ReplayServer/src/com/epicgames/replayserver/MongoDB.java
+++ b/Engine/Source/Programs/ReplayServer/src/com/epicgames/replayserver/MongoDB.java
@@ -37,6 +37,7 @@ public class MongoDB implements BaseDB
 	static DBCollection replayColl	= null;
 	static DBCollection viewerColl	= null;
 	static DBCollection eventColl	= null;
+	static DBCollection recentColl	= null;
 	static DBCollection logColl		= null;
 
 	public MongoDB() throws UnknownHostException
@@ -56,6 +57,7 @@ public class MongoDB implements BaseDB
 		replayColl 	= db.getCollection( ReplayProps.getString( "mongoDB_ReplayCollection", "replayCollection" ) );
 		viewerColl 	= db.getCollection( ReplayProps.getString( "mongoDB_ViewerCollection", "viewerCollection" ) );
 		eventColl 	= db.getCollection( ReplayProps.getString( "mongoDB_EventCollection", "eventCollection" ) );
+		recentColl 	= db.getCollection( ReplayProps.getString( "mongoDB_RecentCollection", "recentCollection" ) );
 		logColl 	= db.getCollection( ReplayProps.getString( "mongoDB_LogCollection", "logCollection" ) );
 				
 		if ( bResetIndexes )
@@ -80,6 +82,8 @@ public class MongoDB implements BaseDB
 			//eventColl.createIndex( new BasicDBObject( "group" , 1 ), new BasicDBObject( "name", "EGroupIndex" ) );			
 			eventColl.createIndex( new BasicDBObject( "modified", 1 ), new BasicDBObject( "name", "EModifiedIndex" ) );
 
+			recentColl.createIndex( new BasicDBObject( "user", 1 ), new BasicDBObject( "name", "RecentUserIndex" ) );
+
 			//logColl.createIndex( new BasicDBObject( "level", 1 ), new BasicDBObject( "name", "LevelIndex" ) );
 			logColl.createIndex( new BasicDBObject( "level", 1 ).append( "created", 1 ), new BasicDBObject( "name", "LevelIndex" ) );
 			logColl.createIndex( new BasicDBObject( "created", 1 ), new BasicDBObject( "name", "LCreatedIndex" ) );
@@ -232,7 +236,7 @@ public class MongoDB implements BaseDB
 							.append( "sizeInBytes", 	(int)0 );
 
 		
-			replayColl.insert( WriteConcern.SAFE, newSession );
+			replayColl.insert( WriteConcern.ACKNOWLEDGED, newSession );
 		}
 		catch ( Exception e )
 		{
@@ -462,7 +466,7 @@ public class MongoDB implements BaseDB
 				obj.put( "meta", eventMeta );
 			}
 
-			eventColl.insert( WriteConcern.SAFE, obj );
+			eventColl.insert( WriteConcern.ACKNOWLEDGED, obj );
         }
         catch ( Exception e )
         {
@@ -516,7 +520,7 @@ public class MongoDB implements BaseDB
         }
         catch ( Exception e )
         {
-        	throw new ReplayException( "MongoDB.EnumerateEvents: There was an error reading the event" );
+        	throw new ReplayException( "MongoDB.enumerateEvents: There was an error reading the event" );
         }
 	}
 
@@ -525,20 +529,65 @@ public class MongoDB implements BaseDB
 		return viewerColl.count( getSpecificViewerQuery( sessionName, viewerName ) ) > 0;
 	}
 
-	public void createViewer( final String sessionName, final String viewerName ) throws ReplayException
+	public void createViewer( final String sessionName, final String viewerName, final String userName ) throws ReplayException
 	{
 		if ( viewerExists( sessionName, viewerName ) )
 		{
-			throw new ReplayException( "MongoDB.CreateViewer: Viewer already exists." );        
+			throw new ReplayException( "MongoDB.createViewer: Viewer already exists." );        
 		}
 		
 		try
 		{
-			viewerColl.insert( WriteConcern.SAFE, getSpecificViewerQuery( sessionName, viewerName ).append( "modified", System.currentTimeMillis() ) );		
+			viewerColl.insert( WriteConcern.ACKNOWLEDGED, getSpecificViewerQuery( sessionName, viewerName ).append( "modified", System.currentTimeMillis() ) );
+		}
+		catch ( Exception e )
+		{
+			throw new ReplayException( "MongoDB.createViewer: Insert failed." );        
+		}
+		
+		addToRecentList( sessionName, userName );
+	}
+	
+	public void addToRecentList( final String sessionName, final String userName ) throws ReplayException
+	{
+		if ( userName == null || userName.equals( "" ) )
+		{
+			ReplayLogger.log( Level.WARNING, "addToRecentList: Ignoring null name." );
+			return;
+		}
+
+		try
+		{
+	 		// This is probably slow, but is proof of concept feature, need to make faster
+			final int count 	= (int)recentColl.count( new BasicDBObject( "user", userName ) ) + 1;
+	 		final int maxCount	= 2;
+	 		
+	 		if ( count > maxCount )
+	 		{
+	 			final int numToRemove = count - maxCount;
+
+	 			ReplayLogger.log( Level.INFO, "addToRecentList: Removing: " + numToRemove );
+	 			
+	 			DBCursor cursor = recentColl.find( new BasicDBObject( "user", userName ) );
+
+		 		cursor.sort( new BasicDBObject( "_id", -1 ) );
+		 		cursor.limit( numToRemove );
+ 		
+	 			while ( cursor.hasNext() ) 
+	 			{
+	 				recentColl.remove( cursor.next() );
+	 			}
+	 		}
+	 		
+			final BasicDBObject newObj = new BasicDBObject( "user", userName ).append( "session", UUID.fromString( sessionName ) );
+
+			ReplayLogger.log( Level.INFO, "addToRecentList: Adding: " + sessionName + " : " + userName );
+			
+			recentColl.insert( WriteConcern.ACKNOWLEDGED, newObj );
 		}
 		catch ( Exception e )
 		{
-			throw new ReplayException( "MongoDB.CreateViewer: Insert failed." );        
+			throw new ReplayException( "MongoDB.addToRecentList: Insert failed." );        
 		}
 	}
 
@@ -689,7 +738,7 @@ public class MongoDB implements BaseDB
 		    }
 		    catch ( Exception e )
 		    {
-				ReplayLogger.log( Level.SEVERE, "MongoDB.DiscoverSessions. Exception while reading. Deleting: " + getSessionStringSafe( doc ) );
+				ReplayLogger.log( Level.SEVERE, "MongoDB.discoverSessions. Exception while reading. Deleting: " + getSessionStringSafe( doc ) );
 				safeDeleteSessionWithDocument( getSessionStringSafe( doc ), doc );
 		    }
 		}				 
@@ -697,6 +746,61 @@ public class MongoDB implements BaseDB
 		return sessions;
 	}
 
+	public List<ReplaySessionInfo> getRecentSessions( final String appName, final int version, final int changelist, final String userName, final int limit )
+	{
+		List<ReplaySessionInfo> sessions = new ArrayList<ReplaySessionInfo>();
+		
+ 		DBCursor cursor = recentColl.find( new BasicDBObject( "user", userName ) );
+ 		
+ 		cursor.sort( new BasicDBObject( "_id", -1 ) );
+	
+ 		if ( limit != 0 )
+		{
+ 			cursor.limit( limit );
+		}
+
+		while ( cursor.hasNext() ) 
+		{
+		    final DBObject doc = (DBObject)cursor.next();
+
+		    try
+		    {	    	
+		    	final String sessionName = getSessionString( doc );
+		    	
+		    	final DBObject sessionObj = replayColl.findOne( getSessionQuery( sessionName ) );
+		    	
+		    	if ( sessionObj == null )
+		    	{
+					ReplayLogger.log( Level.WARNING, "MongoDB.getRecentSessions. Session not found: " + sessionName );
+					recentColl.remove( doc );
+		    		continue;
+		    	}
+
+		    	ReplaySessionInfo sessionInfo = getSessionInfoFromDoc( sessionObj );
+
+		    	// FIXME: Make this part of the insert/search/schema flow
+		    	if ( version != 0 && sessionInfo.version != version )
+		    	{
+		    		continue;
+		    	}
+
+		    	if ( appName != null && !sessionInfo.appName.equals( appName ) )
+		    	{
+		    		continue;
+		    	}
+
+		    	sessions.add( sessionInfo );
+		    }
+		    catch ( Exception e )
+		    {
+				ReplayLogger.log( Level.SEVERE, "MongoDB.getRecentSessions. Exception while reading. Deleting: " + getSessionStringSafe( doc ) );
+				recentColl.remove( doc );
+		    }
+		}				 
+
+		return sessions;
+	}
+
 	public void log( final Level level, final String str )
 	{
 		BasicDBObject logObj = new BasicDBObject( "level", level.intValue() ).append( "created", System.currentTimeMillis() ).append( "log", str );
diff --git a/Engine/Source/Programs/ReplayServer/src/com/epicgames/replayserver/ReplayDB.java b/Engine/Source/Programs/ReplayServer/src/com/epicgames/replayserver/ReplayDB.java
--- a/Engine/Source/Programs/ReplayServer/src/com/epicgames/replayserver/ReplayDB.java
+++ b/Engine/Source/Programs/ReplayServer/src/com/epicgames/replayserver/ReplayDB.java
@@ -40,7 +40,7 @@ public class ReplayDB
 				{
 					for ( int t = 0; t < 10; t++ )
 					{
-						createViewer( sessionName );
+						createViewer( sessionName, "fakeUser" );
 					}
 				}
 
@@ -157,14 +157,12 @@ public class ReplayDB
 		refreshSession( SessionName, false, false, NumChunks, DemoTimeInMS, 0 );
 	}
 
-	static String createViewer( final String sessionName ) throws ReplayException
+	static String createViewer( final String sessionName, final String userName ) throws ReplayException
 	{
-		final UUID sessionId = java.util.UUID.randomUUID();
-		final String viewerName = sessionId.toString();
-
 		try
 		{				
-			baseDB.createViewer( sessionName, viewerName );
+			final String viewerName = java.util.UUID.randomUUID().toString();
+			baseDB.createViewer( sessionName, viewerName, userName );
 			return viewerName; 
 		}
 		catch ( ReplayException e )
@@ -237,6 +235,11 @@ public class ReplayDB
 	{
 		return baseDB.discoverSessions( appName, version, changelist, limit );
 	}
+
+	static List<ReplaySessionInfo> getRecentSessions( final String appName, final int version, final int changelist, final String userName, final int limit )
+	{
+		return baseDB.getRecentSessions( appName, version, changelist, userName, limit );
+	}
 	
 	static void log( final Level level, final String str )
 	{
diff --git a/Engine/Source/Programs/ReplayServer/src/com/epicgames/replayserver/StartDownloading.java b/Engine/Source/Programs/ReplayServer/src/com/epicgames/replayserver/StartDownloading.java
--- a/Engine/Source/Programs/ReplayServer/src/com/epicgames/replayserver/StartDownloading.java
+++ b/Engine/Source/Programs/ReplayServer/src/com/epicgames/replayserver/StartDownloading.java
@@ -25,6 +25,7 @@ public class StartDownloading extends HttpServlet
     	try
     	{
     		final String sessionName = request.getParameter( "Session" );
+    		final String userName = request.getParameter( "User" );
 
             if ( !ReplayDB.sessionExists( sessionName, false, false ) )
         	{
@@ -44,7 +45,7 @@ public class StartDownloading extends HttpServlet
                 response.setHeader( "State", "Final" );
             }
 
-            final String viewerName = ReplayDB.createViewer( sessionName );
+            final String viewerName = ReplayDB.createViewer( sessionName, userName );
     		
     		response.setHeader( "NumChunks", "" + session.numChunks );
     		response.setHeader( "Time", "" + session.demoTimeInMS );
