#replayserver: Change max recent replays to 10 + don't insert same replay more than once into recent list + fix sort problem when deleting old recents

diff --git a/Engine/Source/Programs/ReplayServer/src/com/epicgames/replayserver/MongoDB.java b/Engine/Source/Programs/ReplayServer/src/com/epicgames/replayserver/MongoDB.java
--- a/Engine/Source/Programs/ReplayServer/src/com/epicgames/replayserver/MongoDB.java
+++ b/Engine/Source/Programs/ReplayServer/src/com/epicgames/replayserver/MongoDB.java
@@ -558,9 +558,11 @@ public class MongoDB implements BaseDB
 
 		try
 		{
-	 		// This is probably slow, but is proof of concept feature, need to make faster
-			final int count 	= (int)recentColl.count( new BasicDBObject( "user", userName ) ) + 1;
-	 		final int maxCount	= 2;
+			final BasicDBObject query = new BasicDBObject( "user", userName );
+
+			// This is probably slow, but is proof of concept feature, need to make faster
+			final int count 	= (int)recentColl.count( query ) + 1;
+	 		final int maxCount	= 10;
 	 		
 	 		if ( count > maxCount )
 	 		{
@@ -568,9 +570,9 @@ public class MongoDB implements BaseDB
 
 	 			ReplayLogger.log( Level.INFO, "addToRecentList: Removing: " + numToRemove );
 	 			
-	 			DBCursor cursor = recentColl.find( new BasicDBObject( "user", userName ) );
+	 			DBCursor cursor = recentColl.find( query );
 
-		 		cursor.sort( new BasicDBObject( "_id", -1 ) );
+		 		cursor.sort( new BasicDBObject( "_id", 1 ) );		// Sort oldest first
 		 		cursor.limit( numToRemove );
  		
 	 			while ( cursor.hasNext() ) 
@@ -582,8 +584,8 @@ public class MongoDB implements BaseDB
 			final BasicDBObject newObj = new BasicDBObject( "user", userName ).append( "session", UUID.fromString( sessionName ) );
 
 			ReplayLogger.log( Level.INFO, "addToRecentList: Adding: " + sessionName + " : " + userName );
-			
-			recentColl.insert( WriteConcern.ACKNOWLEDGED, newObj );
+
+			recentColl.update( newObj, newObj, true, false );
 		}
 		catch ( Exception e )
 		{
