#replayserver: Add ability for games to add meta tags to replays, and search by them

diff --git a/Engine/Source/Programs/ReplayServer/ReplayServer.Properties b/Engine/Source/Programs/ReplayServer/ReplayServer.Properties
--- a/Engine/Source/Programs/ReplayServer/ReplayServer.Properties
+++ b/Engine/Source/Programs/ReplayServer/ReplayServer.Properties
@@ -26,4 +26,5 @@ mongoDB_cleanupStaleLiveSessions_MaxMinutes=1
 mongoDB_cleanupOldSessions_MaxDays=30
 mongoDB_cleanupStaleViewers_MaxMinutes=1
 mongoDB_cleanupLogs_MaxDays=14
-mongoDB_cleanupLogs_MaxFinestDays=3
\ No newline at end of file
+mongoDB_cleanupLogs_MaxFinestDays=3
+mongoDB_ResetIndexes=0
\ No newline at end of file
diff --git a/Engine/Source/Programs/ReplayServer/src/com/epicgames/replayserver/BaseDB.java b/Engine/Source/Programs/ReplayServer/src/com/epicgames/replayserver/BaseDB.java
--- a/Engine/Source/Programs/ReplayServer/src/com/epicgames/replayserver/BaseDB.java
+++ b/Engine/Source/Programs/ReplayServer/src/com/epicgames/replayserver/BaseDB.java
@@ -49,7 +49,7 @@ public interface BaseDB
 
 	public void shutdown();
 	public boolean sessionExists( final String sessionName );
-	public void createSession( final String appName, final int version, final int changelist, final String sessionName, final String friendlyName ) throws ReplayException;
+	public void createSession( final String appName, final int version, final int changelist, final String sessionName, final String friendlyName, final String metaString ) throws ReplayException;
 	public void deleteSession( final String sessionName ) throws ReplayException;
 	public void refreshSession( final String 	sessionName, 
 								final boolean 	bIsLive, 
@@ -78,7 +78,7 @@ public interface BaseDB
 	public void deleteViewer( final String sessionName, final String viewerName ) throws ReplayException;
 	public ReplaySessionInfo getSessionInfo( final String sessionName ) throws ReplayException;
 	public int getNumSessions( final String appName, final int version, final int changelist );
-	public List<ReplaySessionInfo> discoverSessions( final String appName, final int version, final int changelist, final int limit );
+	public List<ReplaySessionInfo> discoverSessions( final String appName, final int version, final int changelist, final String metaString, final int limit );
 	public List<ReplaySessionInfo> getRecentSessions( final String appName, final int version, final int changelist, final String userName, final int limit );
 
 	public void log( final Level level, final String str );
diff --git a/Engine/Source/Programs/ReplayServer/src/com/epicgames/replayserver/EnumerateSessions.java b/Engine/Source/Programs/ReplayServer/src/com/epicgames/replayserver/EnumerateSessions.java
--- a/Engine/Source/Programs/ReplayServer/src/com/epicgames/replayserver/EnumerateSessions.java
+++ b/Engine/Source/Programs/ReplayServer/src/com/epicgames/replayserver/EnumerateSessions.java
@@ -45,6 +45,7 @@ public class EnumerateSessions extends HttpServlet
     	final String versionString 	= request.getParameter( "Version" );
     	final String clString 		= request.getParameter( "CL" );
     	final String userName 		= request.getParameter( "User" );
+    	final String metaString 	= request.getParameter( "Meta" );
 
     	if ( appName == null )
     	{
@@ -68,7 +69,7 @@ public class EnumerateSessions extends HttpServlet
     	}
     	else
     	{
-    		replaySessions = ReplayDB.discoverSessions( appName, version, changelist, hardLimit );
+    		replaySessions = ReplayDB.discoverSessions( appName, version, changelist, metaString, hardLimit );
     	}
 
 		Gson gson = new GsonBuilder().setDateFormat( "yyyy-MM-dd'T'HH:mm:ss" ).create();
diff --git a/Engine/Source/Programs/ReplayServer/src/com/epicgames/replayserver/FileDB.java b/Engine/Source/Programs/ReplayServer/src/com/epicgames/replayserver/FileDB.java
--- a/Engine/Source/Programs/ReplayServer/src/com/epicgames/replayserver/FileDB.java
+++ b/Engine/Source/Programs/ReplayServer/src/com/epicgames/replayserver/FileDB.java
@@ -99,7 +99,7 @@ public class FileDB implements BaseDB
 		return true;
 	}
 
-	public void createSession( final String appName, final int version, final int changelist, final String sessionName, final String friendlyName ) throws ReplayException
+	public void createSession( final String appName, final int version, final int changelist, final String sessionName, final String friendlyName, final String metaString ) throws ReplayException
 	{
 		if ( sessionExists( sessionName ) )
 		{
@@ -354,7 +354,7 @@ public class FileDB implements BaseDB
 		return 0;
 	}
 
-	public List<ReplaySessionInfo> discoverSessions( final String appName, final int version, final int changelist, final int limit )
+	public List<ReplaySessionInfo> discoverSessions( final String appName, final int version, final int changelist, final String metaString, final int limit )
 	{
 		List<ReplaySessionInfo> sessions = new ArrayList<ReplaySessionInfo>();
 		
@@ -435,7 +435,7 @@ public class FileDB implements BaseDB
 	public void quickCleanup()
 	{
 		// Loop over all of the sessions and clean up old ones
-		for ( final ReplaySessionInfo sessionEntry : discoverSessions( null, 0, 0, 0 ) )
+		for ( final ReplaySessionInfo sessionEntry : discoverSessions( null, 0, 0, null, 0 ) )
 		{		
 			if ( sessionEntry.bIsLive && sessionEntry.GetRefreshTimeInSeconds() >= 60 )
 			{
diff --git a/Engine/Source/Programs/ReplayServer/src/com/epicgames/replayserver/Index.java b/Engine/Source/Programs/ReplayServer/src/com/epicgames/replayserver/Index.java
--- a/Engine/Source/Programs/ReplayServer/src/com/epicgames/replayserver/Index.java
+++ b/Engine/Source/Programs/ReplayServer/src/com/epicgames/replayserver/Index.java
@@ -82,7 +82,7 @@ public class Index extends HttpServlet
 			response.getWriter().println( "<td></td>");
 
 			// Loop over all of the sessions and fill the table
-			final List< ReplaySessionInfo > sessions = ReplayDB.discoverSessions( null, 0, 0, refresh != null ? 10 : 0 );
+			final List< ReplaySessionInfo > sessions = ReplayDB.discoverSessions( null, 0, 0, null, refresh != null ? 10 : 0 );
 
 			response.getWriter().println( "<p>Total Sessions: " + ReplayDB.getNumSessions( null, 0, 0 ) );
 
diff --git a/Engine/Source/Programs/ReplayServer/src/com/epicgames/replayserver/MongoDB.java b/Engine/Source/Programs/ReplayServer/src/com/epicgames/replayserver/MongoDB.java
--- a/Engine/Source/Programs/ReplayServer/src/com/epicgames/replayserver/MongoDB.java
+++ b/Engine/Source/Programs/ReplayServer/src/com/epicgames/replayserver/MongoDB.java
@@ -45,7 +45,7 @@ public class MongoDB implements BaseDB
 		mongoClient = new MongoClient( ReplayProps.getString( "mongoDB_Host", "localhost" ) , ReplayProps.getInt( "mongoDB_Port", "27017" ) );
 
 		final boolean bResetDB 		= false;
-		final boolean bResetIndexes = false;
+		final boolean bResetIndexes = ReplayProps.getInt( "mongoDB_ResetIndexes", "1" ) == 1 ? true : false;
 
 		if ( bResetDB )
 		{
@@ -65,11 +65,13 @@ public class MongoDB implements BaseDB
 			replayColl.dropIndex( "*" );
 			viewerColl.dropIndex( "*" );
 			eventColl.dropIndex( "*" );
+			recentColl.dropIndex( "*" );
 			logColl.dropIndex( "*" );
 	
 			replayColl.createIndex( new BasicDBObject( "session", 1 ), new BasicDBObject( "name", "SessionIndex" ) );
-			replayColl.createIndex( new BasicDBObject( "version", 1 ).append( "app", 1 ).append( "bIsLive", 1 ).append( "created", 1 ), new BasicDBObject( "name", "VersionAppSortIndex" ) );
+			replayColl.createIndex( new BasicDBObject( "app", 1 ).append( "version", 1 ).append( "bIsLive", 1 ).append( "created", 1 ), new BasicDBObject( "name", "VersionAppSortIndex" ) );
 			replayColl.createIndex( new BasicDBObject( "app", 1 ).append( "bIsLive", 1 ).append( "created", 1 ), new BasicDBObject( "name", "AppSortIndex" ) );
+			replayColl.createIndex( new BasicDBObject( "app", 1 ).append( "meta", 1 ), new BasicDBObject( "name", "AppMetaSortIndex" ) );
 			replayColl.createIndex( new BasicDBObject( "created", 1 ), new BasicDBObject( "name", "CreatedIndex" ) );
 			replayColl.createIndex( new BasicDBObject( "modified", 1 ), new BasicDBObject( "name", "ModifiedIndex" ) );
 
@@ -77,14 +79,11 @@ public class MongoDB implements BaseDB
 			viewerColl.createIndex( new BasicDBObject( "viewer", 1 ), new BasicDBObject( "name", "ViewerIndex" ) );
 			viewerColl.createIndex( new BasicDBObject( "modified", 1 ), new BasicDBObject( "name", "VModifiedIndex" ) );
 
+			recentColl.createIndex( new BasicDBObject( "user", 1 ), new BasicDBObject( "name", "RecentUserIndex" ) );
+
 			eventColl.createIndex( new BasicDBObject( "session", 1 ).append( "group", 1 ).append( "time1", 1 ), new BasicDBObject( "name", "ESessionIndex" ) );
-			//eventColl.createIndex( new BasicDBObject( "time1", 1 ), new BasicDBObject( "name", "ETime1Index" ) );
-			//eventColl.createIndex( new BasicDBObject( "group" , 1 ), new BasicDBObject( "name", "EGroupIndex" ) );			
 			eventColl.createIndex( new BasicDBObject( "modified", 1 ), new BasicDBObject( "name", "EModifiedIndex" ) );
 
-			recentColl.createIndex( new BasicDBObject( "user", 1 ), new BasicDBObject( "name", "RecentUserIndex" ) );
-
-			//logColl.createIndex( new BasicDBObject( "level", 1 ), new BasicDBObject( "name", "LevelIndex" ) );
 			logColl.createIndex( new BasicDBObject( "level", 1 ).append( "created", 1 ), new BasicDBObject( "name", "LevelIndex" ) );
 			logColl.createIndex( new BasicDBObject( "created", 1 ), new BasicDBObject( "name", "LCreatedIndex" ) );
 		}
@@ -213,7 +212,7 @@ public class MongoDB implements BaseDB
 		return replayColl.count( getSessionQuery( sessionName ) ) > 0;
 	}
 	
-	public void createSession( final String appName, final int version, final int changelist, final String sessionName, final String friendlyName ) throws ReplayException
+	public void createSession( final String appName, final int version, final int changelist, final String sessionName, final String friendlyName, final String metaString ) throws ReplayException
 	{
 		if ( sessionExists( sessionName ) )
 		{
@@ -233,7 +232,8 @@ public class MongoDB implements BaseDB
 							.append( "bIncomplete", 	(boolean)false )
 							.append( "numChunks", 		(int)0 )
 							.append( "demoTimeInMS", 	(int)0 )
-							.append( "sizeInBytes", 	(int)0 );
+							.append( "sizeInBytes", 	(int)0 )
+							.append( "meta", 			metaString );
 
 		
 			replayColl.insert( WriteConcern.ACKNOWLEDGED, newSession );
@@ -695,26 +695,27 @@ public class MongoDB implements BaseDB
  		return (int)replayColl.count();
 	}
 
-	public List<ReplaySessionInfo> discoverSessions( final String appName, final int version, final int changelist, final int limit )
+	public List<ReplaySessionInfo> discoverSessions( final String appName, final int version, final int changelist, final String metaString, final int limit )
 	{
 		List<ReplaySessionInfo> sessions = new ArrayList<ReplaySessionInfo>();
 	
  		DBCursor cursor = null;
  		
- 		if ( version != 0 )
+ 		if ( appName != null )
  		{
- 	 		if ( appName != null )
- 	 		{
- 	 			cursor = replayColl.find( new BasicDBObject( "version", version ).append( "app", appName ) );
- 	 		}
- 	 		else
- 	 		{
- 	 			cursor = replayColl.find( new BasicDBObject( "version", version ) );
- 	 		}
- 		}
- 		else if ( appName != null )
- 		{
- 			cursor = replayColl.find( new BasicDBObject( "app", appName ) );
+ 			BasicDBObject query = new BasicDBObject( "app", appName );
+ 			
+ 			if ( version != 0 )
+	 		{
+	 			query.append( "version", version );	 			
+	 		}
+
+ 			if ( metaString != null )
+ 			{	 				
+		 		query.append( "meta", metaString );
+ 			}
+
+ 			cursor = replayColl.find( query );
  		}
  		else
  		{
diff --git a/Engine/Source/Programs/ReplayServer/src/com/epicgames/replayserver/ReplayDB.java b/Engine/Source/Programs/ReplayServer/src/com/epicgames/replayserver/ReplayDB.java
--- a/Engine/Source/Programs/ReplayServer/src/com/epicgames/replayserver/ReplayDB.java
+++ b/Engine/Source/Programs/ReplayServer/src/com/epicgames/replayserver/ReplayDB.java
@@ -34,7 +34,7 @@ public class ReplayDB
 			{
 				final String app = "testApp";
 				
-				final String sessionName = createSession( app, 0, 0, "friendlyName" );
+				final String sessionName = createSession( app, 0, 0, "friendlyName", null );
 				
 				if ( i % 10 == 0 )
 				{
@@ -116,14 +116,14 @@ public class ReplayDB
 		return true;
 	}
 
-	static String createSession( final String appName, final int version, final int changelist, final String friendlyName ) throws ReplayException
+	static String createSession( final String appName, final int version, final int changelist, final String friendlyName, final String metaString ) throws ReplayException
 	{
 		final UUID sessionId = java.util.UUID.randomUUID();
 		final String sessionName = sessionId.toString();
 
 		try
 		{				
-			baseDB.createSession( appName, version, changelist, sessionName, friendlyName );
+			baseDB.createSession( appName, version, changelist, sessionName, friendlyName, metaString );
 			return sessionName; 
 		}
 		catch ( ReplayException e )
@@ -231,9 +231,9 @@ public class ReplayDB
 		return baseDB.getNumSessions( appName, version, changelist );
 	}
 
-	static List<ReplaySessionInfo> discoverSessions(  final String appName, final int version, final int changelist, final int limit )
+	static List<ReplaySessionInfo> discoverSessions(  final String appName, final int version, final int changelist, final String metaString, final int limit )
 	{
-		return baseDB.discoverSessions( appName, version, changelist, limit );
+		return baseDB.discoverSessions( appName, version, changelist, metaString, limit );
 	}
 
 	static List<ReplaySessionInfo> getRecentSessions( final String appName, final int version, final int changelist, final String userName, final int limit )
diff --git a/Engine/Source/Programs/ReplayServer/src/com/epicgames/replayserver/StartUploading.java b/Engine/Source/Programs/ReplayServer/src/com/epicgames/replayserver/StartUploading.java
--- a/Engine/Source/Programs/ReplayServer/src/com/epicgames/replayserver/StartUploading.java
+++ b/Engine/Source/Programs/ReplayServer/src/com/epicgames/replayserver/StartUploading.java
@@ -28,6 +28,7 @@ public class StartUploading extends HttpServlet
         	final String versionString 	= request.getParameter( "Version" );
         	final String clString 		= request.getParameter( "CL" );
     		final String friendlyName 	= request.getParameter( "Friendly" );
+    		final String metaString 	= request.getParameter( "Meta" );
     		
     		if ( appName == null )
     		{
@@ -61,9 +62,9 @@ public class StartUploading extends HttpServlet
         	final int changelist = Integer.parseUnsignedInt( clString );
 
     		// Start a live session
-    		final String SessionName = ReplayDB.createSession( appName, version, changelist, friendlyName );
+    		final String SessionName = ReplayDB.createSession( appName, version, changelist, friendlyName, metaString );
       	
-    		ReplayLogger.log( Level.FINE, "StartUploading. Success: " + appName + "/" + versionString + "/" + clString + "/" + SessionName + "/" + friendlyName );
+    		ReplayLogger.log( Level.FINE, "StartUploading. Success: " + appName + "/" + versionString + "/" + clString + "/" + SessionName + "/" + friendlyName + "/" + metaString );
 
     		response.setHeader( "Session",  SessionName );   		
     	}
