#replayserver: Auto-detect changed indexes, and only rebuild them when necessary

diff --git a/Engine/Source/Programs/ReplayServer/src/com/epicgames/replayserver/MongoDB.java b/Engine/Source/Programs/ReplayServer/src/com/epicgames/replayserver/MongoDB.java
--- a/Engine/Source/Programs/ReplayServer/src/com/epicgames/replayserver/MongoDB.java
+++ b/Engine/Source/Programs/ReplayServer/src/com/epicgames/replayserver/MongoDB.java
@@ -45,7 +45,7 @@ public class MongoDB implements BaseDB
 		mongoClient = new MongoClient( ReplayProps.getString( "mongoDB_Host", "localhost" ) , ReplayProps.getInt( "mongoDB_Port", "27017" ) );
 
 		final boolean bResetDB 		= false;
-		final boolean bResetIndexes = ReplayProps.getInt( "mongoDB_ResetIndexes", "1" ) == 1 ? true : false;
+		final boolean bResetIndexes = ReplayProps.getInt( "mongoDB_ResetIndexes", "0" ) == 1 ? true : false;
 
 		if ( bResetDB )
 		{
@@ -62,31 +62,10 @@ public class MongoDB implements BaseDB
 				
 		if ( bResetIndexes )
 		{
-			replayColl.dropIndex( "*" );
-			viewerColl.dropIndex( "*" );
-			eventColl.dropIndex( "*" );
-			recentColl.dropIndex( "*" );
-			logColl.dropIndex( "*" );
-	
-			replayColl.createIndex( new BasicDBObject( "session", 1 ), new BasicDBObject( "name", "SessionIndex" ) );
-			replayColl.createIndex( new BasicDBObject( "app", 1 ).append( "version", 1 ).append( "bIsLive", 1 ).append( "created", 1 ), new BasicDBObject( "name", "VersionAppSortIndex" ) );
-			replayColl.createIndex( new BasicDBObject( "app", 1 ).append( "bIsLive", 1 ).append( "created", 1 ), new BasicDBObject( "name", "AppSortIndex" ) );
-			replayColl.createIndex( new BasicDBObject( "app", 1 ).append( "meta", 1 ), new BasicDBObject( "name", "AppMetaSortIndex" ) );
-			replayColl.createIndex( new BasicDBObject( "created", 1 ), new BasicDBObject( "name", "CreatedIndex" ) );
-			replayColl.createIndex( new BasicDBObject( "modified", 1 ), new BasicDBObject( "name", "ModifiedIndex" ) );
-
-			viewerColl.createIndex( new BasicDBObject( "session", 1 ), new BasicDBObject( "name", "VSessionIndex" ) );
-			viewerColl.createIndex( new BasicDBObject( "viewer", 1 ), new BasicDBObject( "name", "ViewerIndex" ) );
-			viewerColl.createIndex( new BasicDBObject( "modified", 1 ), new BasicDBObject( "name", "VModifiedIndex" ) );
-
-			recentColl.createIndex( new BasicDBObject( "user", 1 ), new BasicDBObject( "name", "RecentUserIndex" ) );
-
-			eventColl.createIndex( new BasicDBObject( "session", 1 ).append( "group", 1 ).append( "time1", 1 ), new BasicDBObject( "name", "ESessionIndex" ) );
-			eventColl.createIndex( new BasicDBObject( "modified", 1 ), new BasicDBObject( "name", "EModifiedIndex" ) );
-
-			logColl.createIndex( new BasicDBObject( "level", 1 ).append( "created", 1 ), new BasicDBObject( "name", "LevelIndex" ) );
-			logColl.createIndex( new BasicDBObject( "created", 1 ), new BasicDBObject( "name", "LCreatedIndex" ) );
+			ResetIndexes();
 		}
+	
+		RegisterIndexes();
 
 		/*
 		listDocuments( replayColl );
@@ -102,6 +81,50 @@ public class MongoDB implements BaseDB
 		PrintCollectionStats( db.getCollection( "replayFS.files" ) );
 		PrintCollectionStats( db.getCollection( "replayFS.chunks" ) );
 	}
+	
+	public void ResetIndexes()
+	{
+		replayColl.dropIndex( "*" );
+		viewerColl.dropIndex( "*" );
+		eventColl.dropIndex( "*" );
+		recentColl.dropIndex( "*" );
+		logColl.dropIndex( "*" );		
+	}
+
+	public void RegisterIndexes()
+	{
+		EnsureIndex( replayColl, new BasicDBObject( "session", 1 ), "SessionIndex" );
+		EnsureIndex( replayColl, new BasicDBObject( "app", 1 ).append( "version", 1 ).append( "bIsLive", 1 ).append( "created", 1 ), "VersionAppSortIndex" );
+		EnsureIndex( replayColl, new BasicDBObject( "app", 1 ).append( "bIsLive", 1 ).append( "created", 1 ), "AppSortIndex" );
+		EnsureIndex( replayColl, new BasicDBObject( "app", 1 ).append( "meta", 1 ), "AppMetaSortIndex" );
+		EnsureIndex( replayColl, new BasicDBObject( "created", 1 ), "CreatedIndex" );
+		EnsureIndex( replayColl, new BasicDBObject( "modified", 1 ), "ModifiedIndex" );
+
+		EnsureIndex( viewerColl, new BasicDBObject( "session", 1 ), "VSessionIndex" );
+		EnsureIndex( viewerColl, new BasicDBObject( "viewer", 1 ), "ViewerIndex" );
+		EnsureIndex( viewerColl, new BasicDBObject( "modified", 1 ), "VModifiedIndex" );
+
+		EnsureIndex( recentColl, new BasicDBObject( "user", 1 ), "RecentUserIndex" );
+
+		EnsureIndex( eventColl, new BasicDBObject( "session", 1 ).append( "group", 1 ).append( "time1", 1 ), "ESessionIndex" );
+		EnsureIndex( eventColl, new BasicDBObject( "modified", 1 ), "EModifiedIndex" );
+
+		EnsureIndex( logColl, new BasicDBObject( "level", 1 ).append( "created", 1 ), "LevelIndex" );
+		EnsureIndex( logColl, new BasicDBObject( "created", 1 ), "LCreatedIndex" );
+	}
+	
+	public void EnsureIndex( DBCollection collection, DBObject keys, String name )
+	{
+		try
+		{
+			collection.createIndex( keys, name );
+		}
+		catch( Exception e )
+		{
+			replayColl.dropIndex( name );
+			collection.createIndex( keys, name );
+		}
+	}
 
 	public void shutdown()
 	{
