#replayserver: Add support for storing users

diff --git a/Engine/Source/Programs/ReplayServer/src/com/epicgames/replayserver/BaseDB.java b/Engine/Source/Programs/ReplayServer/src/com/epicgames/replayserver/BaseDB.java
--- a/Engine/Source/Programs/ReplayServer/src/com/epicgames/replayserver/BaseDB.java
+++ b/Engine/Source/Programs/ReplayServer/src/com/epicgames/replayserver/BaseDB.java
@@ -49,7 +49,7 @@ public interface BaseDB
 
 	public void shutdown();
 	public boolean sessionExists( final String sessionName );
-	public void createSession( final String appName, final int version, final int changelist, final String sessionName, final String friendlyName, final String metaString ) throws ReplayException;
+	public void createSession( final String appName, final int version, final int changelist, final String sessionName, final String friendlyName, final List< String > userNames, final String metaString ) throws ReplayException;
 	public void deleteSession( final String sessionName ) throws ReplayException;
 	public void refreshSession( final String 	sessionName, 
 								final boolean 	bIsLive, 
diff --git a/Engine/Source/Programs/ReplayServer/src/com/epicgames/replayserver/FileDB.java b/Engine/Source/Programs/ReplayServer/src/com/epicgames/replayserver/FileDB.java
--- a/Engine/Source/Programs/ReplayServer/src/com/epicgames/replayserver/FileDB.java
+++ b/Engine/Source/Programs/ReplayServer/src/com/epicgames/replayserver/FileDB.java
@@ -99,7 +99,7 @@ public class FileDB implements BaseDB
 		return true;
 	}
 
-	public void createSession( final String appName, final int version, final int changelist, final String sessionName, final String friendlyName, final String metaString ) throws ReplayException
+	public void createSession( final String appName, final int version, final int changelist, final String sessionName, final String friendlyName, final List< String > userNames, final String metaString ) throws ReplayException
 	{
 		if ( sessionExists( sessionName ) )
 		{
diff --git a/Engine/Source/Programs/ReplayServer/src/com/epicgames/replayserver/MongoDB.java b/Engine/Source/Programs/ReplayServer/src/com/epicgames/replayserver/MongoDB.java
--- a/Engine/Source/Programs/ReplayServer/src/com/epicgames/replayserver/MongoDB.java
+++ b/Engine/Source/Programs/ReplayServer/src/com/epicgames/replayserver/MongoDB.java
@@ -9,6 +9,7 @@ import java.io.PrintWriter;
 import java.net.UnknownHostException;
 import java.text.SimpleDateFormat;
 import java.util.ArrayList;
+//import java.util.Collection;
 import java.util.Date;
 import java.util.List;
 import java.util.UUID;
@@ -97,6 +98,7 @@ public class MongoDB implements BaseDB
 		EnsureIndex( replayColl, new BasicDBObject( "app", 1 ).append( "version", 1 ).append( "bIsLive", 1 ).append( "created", 1 ), "VersionAppSortIndex" );
 		EnsureIndex( replayColl, new BasicDBObject( "app", 1 ).append( "bIsLive", 1 ).append( "created", 1 ), "AppSortIndex" );
 		EnsureIndex( replayColl, new BasicDBObject( "app", 1 ).append( "meta", 1 ), "AppMetaSortIndex" );
+		EnsureIndex( replayColl, new BasicDBObject( "app", 1 ).append( "users", 1 ).append( "created", 1 ), "AppUserIndex" );
 		EnsureIndex( replayColl, new BasicDBObject( "created", 1 ), "CreatedIndex" );
 		EnsureIndex( replayColl, new BasicDBObject( "modified", 1 ), "ModifiedIndex" );
 
@@ -235,7 +237,7 @@ public class MongoDB implements BaseDB
 		return replayColl.count( getSessionQuery( sessionName ) ) > 0;
 	}
 	
-	public void createSession( final String appName, final int version, final int changelist, final String sessionName, final String friendlyName, final String metaString ) throws ReplayException
+	public void createSession( final String appName, final int version, final int changelist, final String sessionName, final String friendlyName, final List< String > userNames, final String metaString ) throws ReplayException
 	{
 		if ( sessionExists( sessionName ) )
 		{
@@ -256,6 +258,7 @@ public class MongoDB implements BaseDB
 							.append( "numChunks", 		(int)0 )
 							.append( "demoTimeInMS", 	(int)0 )
 							.append( "sizeInBytes", 	(int)0 )
+							.append( "users", 			userNames )
 							.append( "meta", 			metaString );
 
 		
@@ -346,7 +349,7 @@ public class MongoDB implements BaseDB
 		WriteResult result = replayColl.update( getSessionQuery( sessionName ), new BasicDBObject( "$set", setData ).append( "$inc", incData ), true, false );
         if ( result.getN() != 1 )
         {
-        	throw new ReplayException( "MongoDB.UpdateSessionMetaData: result.getN() != 1" );
+        	throw new ReplayException( "MongoDB.refreshSession: result.getN() != 1" );
         }
 	}
 
@@ -358,7 +361,20 @@ public class MongoDB implements BaseDB
 		WriteResult result = replayColl.update( getSessionQuery( sessionName ), new BasicDBObject( "$set", setData ).append( "$inc", incData ), true, false );
 		if ( result.getN() != 1 )
 		{
-			throw new ReplayException( "MongoDB.UpdateSessionMetaData: result.getN() != 1" );
+			throw new ReplayException( "MongoDB.refreshSession: result.getN() != 1" );
+		}
+	}
+
+	public void addUsersToSession( final String sessionName, final List< String > userNames ) throws ReplayException
+	{
+		final BasicDBObject setData 		= new BasicDBObject( "modified", System.currentTimeMillis() );
+		final BasicDBObject addToSetData 	= new BasicDBObject( "$each", userNames ); 
+
+		WriteResult result = replayColl.update( getSessionQuery( sessionName ), new BasicDBObject( "$set", setData ).append( "$addToSet", addToSetData ), true, false );
+
+		if ( result.getN() != 1 )
+		{
+			throw new ReplayException( "MongoDB.addUserToSession: result.getN() != 1" );
 		}
 	}
 
@@ -668,8 +684,35 @@ public class MongoDB implements BaseDB
 		return obj;
 	}
 
+	/*
+	public static <T> List<T> castList( Class<? extends T> clazz, Collection<?> c ) 
+	{
+	    List<T> r = new ArrayList<T>( c.size() );
+
+	    for ( Object o : c )
+	    {
+	    	r.add( clazz.cast( o ) );
+	    }
+
+	    return r;
+	}
+	*/
+
 	private ReplaySessionInfo getSessionInfoFromDoc( final DBObject doc ) throws ReplayException
 	{
+    	/*
+		try
+    	{
+    		List<String> users = castList( String.class, (List<?>)doc.get( "users" ) );
+
+    	    ReplayLogger.log( Level.INFO, "Users: " + users );
+    	}
+    	catch( Exception e )
+    	{
+    	    ReplayLogger.log( Level.INFO, "NULL Users ");
+    	}
+    	*/
+
 	    try
 	    {
 	    	final String sessionName = getSessionString( doc );
@@ -738,6 +781,13 @@ public class MongoDB implements BaseDB
 		 		query.append( "meta", metaString );
  			}
 
+ 			/*
+ 			if ( userString != null )
+ 			{
+ 	 			query.append( "users", userString );
+ 			}
+ 			*/
+ 			
  			cursor = replayColl.find( query );
  		}
  		else
diff --git a/Engine/Source/Programs/ReplayServer/src/com/epicgames/replayserver/ReplayDB.java b/Engine/Source/Programs/ReplayServer/src/com/epicgames/replayserver/ReplayDB.java
--- a/Engine/Source/Programs/ReplayServer/src/com/epicgames/replayserver/ReplayDB.java
+++ b/Engine/Source/Programs/ReplayServer/src/com/epicgames/replayserver/ReplayDB.java
@@ -34,7 +34,7 @@ public class ReplayDB
 			{
 				final String app = "testApp";
 				
-				final String sessionName = createSession( app, 0, 0, "friendlyName", null );
+				final String sessionName = createSession( app, 0, 0, "friendlyName", null, null );
 				
 				if ( i % 10 == 0 )
 				{
@@ -116,14 +116,14 @@ public class ReplayDB
 		return true;
 	}
 
-	static String createSession( final String appName, final int version, final int changelist, final String friendlyName, final String metaString ) throws ReplayException
+	static String createSession( final String appName, final int version, final int changelist, final String friendlyName, final List< String > userNames, final String metaString ) throws ReplayException
 	{
 		final UUID sessionId = java.util.UUID.randomUUID();
 		final String sessionName = sessionId.toString();
 
 		try
 		{				
-			baseDB.createSession( appName, version, changelist, sessionName, friendlyName, metaString );
+			baseDB.createSession( appName, version, changelist, sessionName, friendlyName, userNames, metaString );
 			return sessionName; 
 		}
 		catch ( ReplayException e )
diff --git a/Engine/Source/Programs/ReplayServer/src/com/epicgames/replayserver/StartUploading.java b/Engine/Source/Programs/ReplayServer/src/com/epicgames/replayserver/StartUploading.java
--- a/Engine/Source/Programs/ReplayServer/src/com/epicgames/replayserver/StartUploading.java
+++ b/Engine/Source/Programs/ReplayServer/src/com/epicgames/replayserver/StartUploading.java
@@ -4,6 +4,8 @@
 package com.epicgames.replayserver;
 
 import java.io.IOException;
+import java.io.InputStreamReader;
+import java.util.List;
 import java.util.logging.Level;
 
 import javax.servlet.ServletException;
@@ -11,11 +13,17 @@ import javax.servlet.http.HttpServlet;
 import javax.servlet.http.HttpServletRequest;
 import javax.servlet.http.HttpServletResponse;
 
+import com.google.gson.Gson;
 
 public class StartUploading extends HttpServlet 
 {
 	private static final long serialVersionUID = 1L;
 
+	public class ReplayUsers 
+	{
+		public List< String > users;
+	}
+
 	public StartUploading()
 	{		
 	}
@@ -61,8 +69,25 @@ public class StartUploading extends HttpServlet
     		final int version = Integer.parseUnsignedInt( versionString );
         	final int changelist = Integer.parseUnsignedInt( clString );
 
-    		// Start a live session
-    		final String SessionName = ReplayDB.createSession( appName, version, changelist, friendlyName, metaString );
+        	List< String > userNames = null;
+       
+        	if ( request.getContentLength() > 0 )
+        	{
+            	try ( InputStreamReader inputStreamReader = new InputStreamReader( request.getInputStream() ) )
+        		{
+            		ReplayUsers replayUsers = new Gson().fromJson( inputStreamReader, ReplayUsers.class );
+            		userNames = replayUsers.users;
+        		}
+            	catch( Exception e )
+            	{
+            		ReplayLogger.log( Level.SEVERE, "StartUploading: Stream failed: " + appName + "/" + version );
+                    response.sendError( HttpServletResponse.SC_BAD_REQUEST );
+                    return;
+            	}
+        	}
+
+            // Start a live session
+    		final String SessionName = ReplayDB.createSession( appName, version, changelist, friendlyName, userNames, metaString );
       	
     		ReplayLogger.log( Level.FINE, "StartUploading. Success: " + appName + "/" + versionString + "/" + clString + "/" + SessionName + "/" + friendlyName + "/" + metaString );
 
