#replayserver: Add ability to search for replay by user

diff --git a/Engine/Source/Programs/ReplayServer/src/com/epicgames/replayserver/BaseDB.java b/Engine/Source/Programs/ReplayServer/src/com/epicgames/replayserver/BaseDB.java
--- a/Engine/Source/Programs/ReplayServer/src/com/epicgames/replayserver/BaseDB.java
+++ b/Engine/Source/Programs/ReplayServer/src/com/epicgames/replayserver/BaseDB.java
@@ -78,7 +78,7 @@ public interface BaseDB
 	public void deleteViewer( final String sessionName, final String viewerName ) throws ReplayException;
 	public ReplaySessionInfo getSessionInfo( final String sessionName ) throws ReplayException;
 	public int getNumSessions( final String appName, final int version, final int changelist );
-	public List<ReplaySessionInfo> discoverSessions( final String appName, final int version, final int changelist, final String metaString, final int limit );
+	public List<ReplaySessionInfo> discoverSessions( final String appName, final int version, final int changelist, final String userString, final String metaString, final int limit );
 	public List<ReplaySessionInfo> getRecentSessions( final String appName, final int version, final int changelist, final String userName, final int limit );
 
 	public void log( final Level level, final String str );
diff --git a/Engine/Source/Programs/ReplayServer/src/com/epicgames/replayserver/EnumerateSessions.java b/Engine/Source/Programs/ReplayServer/src/com/epicgames/replayserver/EnumerateSessions.java
--- a/Engine/Source/Programs/ReplayServer/src/com/epicgames/replayserver/EnumerateSessions.java
+++ b/Engine/Source/Programs/ReplayServer/src/com/epicgames/replayserver/EnumerateSessions.java
@@ -44,6 +44,7 @@ public class EnumerateSessions extends HttpServlet
     	final String appName 		= request.getParameter( "App" );
     	final String versionString 	= request.getParameter( "Version" );
     	final String clString 		= request.getParameter( "CL" );
+    	final String recentName 	= request.getParameter( "Recent" );
     	final String userName 		= request.getParameter( "User" );
     	final String metaString 	= request.getParameter( "Meta" );
 
@@ -63,13 +64,13 @@ public class EnumerateSessions extends HttpServlet
 
     	List< ReplaySessionInfo > replaySessions = null;
     	
-    	if ( userName != null )
+    	if ( recentName != null )
     	{
-    		replaySessions = ReplayDB.getRecentSessions( appName, version, changelist, userName, hardLimit );
+    		replaySessions = ReplayDB.getRecentSessions( appName, version, changelist, recentName, hardLimit );
     	}
     	else
     	{
-    		replaySessions = ReplayDB.discoverSessions( appName, version, changelist, metaString, hardLimit );
+    		replaySessions = ReplayDB.discoverSessions( appName, version, changelist, userName, metaString, hardLimit );
     	}
 
 		Gson gson = new GsonBuilder().setDateFormat( "yyyy-MM-dd'T'HH:mm:ss" ).create();
diff --git a/Engine/Source/Programs/ReplayServer/src/com/epicgames/replayserver/FileDB.java b/Engine/Source/Programs/ReplayServer/src/com/epicgames/replayserver/FileDB.java
--- a/Engine/Source/Programs/ReplayServer/src/com/epicgames/replayserver/FileDB.java
+++ b/Engine/Source/Programs/ReplayServer/src/com/epicgames/replayserver/FileDB.java
@@ -354,7 +354,7 @@ public class FileDB implements BaseDB
 		return 0;
 	}
 
-	public List<ReplaySessionInfo> discoverSessions( final String appName, final int version, final int changelist, final String metaString, final int limit )
+	public List<ReplaySessionInfo> discoverSessions( final String appName, final int version, final int changelist, final String userString, final String metaString, final int limit )
 	{
 		List<ReplaySessionInfo> sessions = new ArrayList<ReplaySessionInfo>();
 		
@@ -435,7 +435,7 @@ public class FileDB implements BaseDB
 	public void quickCleanup()
 	{
 		// Loop over all of the sessions and clean up old ones
-		for ( final ReplaySessionInfo sessionEntry : discoverSessions( null, 0, 0, null, 0 ) )
+		for ( final ReplaySessionInfo sessionEntry : discoverSessions( null, 0, 0, null, null, 0 ) )
 		{		
 			if ( sessionEntry.bIsLive && sessionEntry.GetRefreshTimeInSeconds() >= 60 )
 			{
diff --git a/Engine/Source/Programs/ReplayServer/src/com/epicgames/replayserver/Index.java b/Engine/Source/Programs/ReplayServer/src/com/epicgames/replayserver/Index.java
--- a/Engine/Source/Programs/ReplayServer/src/com/epicgames/replayserver/Index.java
+++ b/Engine/Source/Programs/ReplayServer/src/com/epicgames/replayserver/Index.java
@@ -82,7 +82,7 @@ public class Index extends HttpServlet
 			response.getWriter().println( "<td></td>");
 
 			// Loop over all of the sessions and fill the table
-			final List< ReplaySessionInfo > sessions = ReplayDB.discoverSessions( null, 0, 0, null, refresh != null ? 10 : 0 );
+			final List< ReplaySessionInfo > sessions = ReplayDB.discoverSessions( null, 0, 0, null, null, refresh != null ? 10 : 0 );
 
 			response.getWriter().println( "<p>Total Sessions: " + ReplayDB.getNumSessions( null, 0, 0 ) );
 
diff --git a/Engine/Source/Programs/ReplayServer/src/com/epicgames/replayserver/MongoDB.java b/Engine/Source/Programs/ReplayServer/src/com/epicgames/replayserver/MongoDB.java
--- a/Engine/Source/Programs/ReplayServer/src/com/epicgames/replayserver/MongoDB.java
+++ b/Engine/Source/Programs/ReplayServer/src/com/epicgames/replayserver/MongoDB.java
@@ -761,7 +761,7 @@ public class MongoDB implements BaseDB
  		return (int)replayColl.count();
 	}
 
-	public List<ReplaySessionInfo> discoverSessions( final String appName, final int version, final int changelist, final String metaString, final int limit )
+	public List<ReplaySessionInfo> discoverSessions( final String appName, final int version, final int changelist, final String userString, final String metaString, final int limit )
 	{
 		List<ReplaySessionInfo> sessions = new ArrayList<ReplaySessionInfo>();
 	
@@ -781,12 +781,10 @@ public class MongoDB implements BaseDB
 		 		query.append( "meta", metaString );
  			}
 
- 			/*
  			if ( userString != null )
  			{
  	 			query.append( "users", userString );
  			}
- 			*/
  			
  			cursor = replayColl.find( query );
  		}
diff --git a/Engine/Source/Programs/ReplayServer/src/com/epicgames/replayserver/ReplayDB.java b/Engine/Source/Programs/ReplayServer/src/com/epicgames/replayserver/ReplayDB.java
--- a/Engine/Source/Programs/ReplayServer/src/com/epicgames/replayserver/ReplayDB.java
+++ b/Engine/Source/Programs/ReplayServer/src/com/epicgames/replayserver/ReplayDB.java
@@ -231,9 +231,9 @@ public class ReplayDB
 		return baseDB.getNumSessions( appName, version, changelist );
 	}
 
-	static List<ReplaySessionInfo> discoverSessions(  final String appName, final int version, final int changelist, final String metaString, final int limit )
+	static List<ReplaySessionInfo> discoverSessions(  final String appName, final int version, final int changelist, final String userString, final String metaString, final int limit )
 	{
-		return baseDB.discoverSessions( appName, version, changelist, metaString, limit );
+		return baseDB.discoverSessions( appName, version, changelist, userString, metaString, limit );
 	}
 
 	static List<ReplaySessionInfo> getRecentSessions( final String appName, final int version, final int changelist, final String userName, final int limit )
