#replayserver: Convert replay app name to lowercase

diff --git a/Engine/Source/Programs/ReplayServer/ReplayServer.Properties b/Engine/Source/Programs/ReplayServer/ReplayServer.Properties
--- a/Engine/Source/Programs/ReplayServer/ReplayServer.Properties
+++ b/Engine/Source/Programs/ReplayServer/ReplayServer.Properties
@@ -27,4 +27,5 @@ mongoDB_cleanupOldSessions_MaxDays=30
 mongoDB_cleanupStaleViewers_MaxMinutes=1
 mongoDB_cleanupLogs_MaxDays=14
 mongoDB_cleanupLogs_MaxFinestDays=3
-mongoDB_ResetIndexes=0
\ No newline at end of file
+mongoDB_ResetIndexes=0
+mongoDB_ConvertReplaysToLower=0
\ No newline at end of file
diff --git a/Engine/Source/Programs/ReplayServer/src/com/epicgames/replayserver/MongoDB.java b/Engine/Source/Programs/ReplayServer/src/com/epicgames/replayserver/MongoDB.java
--- a/Engine/Source/Programs/ReplayServer/src/com/epicgames/replayserver/MongoDB.java
+++ b/Engine/Source/Programs/ReplayServer/src/com/epicgames/replayserver/MongoDB.java
@@ -81,6 +81,11 @@ public class MongoDB implements BaseDB
 		PrintCollectionStats( logColl );
 		PrintCollectionStats( db.getCollection( "replayFS.files" ) );
 		PrintCollectionStats( db.getCollection( "replayFS.chunks" ) );
+		
+		if ( ReplayProps.getInt( "mongoDB_ConvertReplaysToLower", "0" ) == 1 )
+		{
+			convertReplaysToLowercase();
+		}
 	}
 	
 	public void ResetIndexes()
@@ -221,6 +226,27 @@ public class MongoDB implements BaseDB
 		}		
 	}
 	
+	void convertReplaysToLowercase()
+	{
+		DBCursor cursor = replayColl.find();
+		
+		while ( cursor.hasNext() ) 
+		{
+		    DBObject doc = (DBObject)cursor.next();
+		    
+		    String appName = (String)doc.get( "app" );
+		    
+		    if ( appName != null )
+		    {
+		    	appName = appName.toLowerCase();
+		    	
+		    	doc.put( "app", appName );
+		    	
+		    	replayColl.save( doc );
+		    }
+		}		
+	}
+
 	static void listFiles()
 	{
 		ReplayLogger.log( "Files:" );
@@ -247,7 +273,7 @@ public class MongoDB implements BaseDB
 		try
 		{
 			final BasicDBObject newSession = getSessionQuery( sessionName )
-							.append( "app", 			appName )
+							.append( "app", 			appName.toLowerCase() )
 							.append( "version", 		version )
 							.append( "cl", 				changelist )
 							.append( "friendlyName", 	friendlyName )
@@ -769,7 +795,7 @@ public class MongoDB implements BaseDB
  		
  		if ( appName != null )
  		{
- 			BasicDBObject query = new BasicDBObject( "app", appName );
+ 			BasicDBObject query = new BasicDBObject( "app", appName.toLowerCase() );
  			
  			if ( version != 0 )
 	 		{
