#replayserver: Change viewer timeout to 3 minutes

diff --git a/Engine/Source/Programs/ReplayServer/ReplayServer.Properties b/Engine/Source/Programs/ReplayServer/ReplayServer.Properties
--- a/Engine/Source/Programs/ReplayServer/ReplayServer.Properties
+++ b/Engine/Source/Programs/ReplayServer/ReplayServer.Properties
@@ -24,7 +24,7 @@ mongoDB_LogCollection=logCollection
 
 mongoDB_cleanupStaleLiveSessions_MaxMinutes=1
 mongoDB_cleanupOldSessions_MaxDays=30
-mongoDB_cleanupStaleViewers_MaxMinutes=1
+mongoDB_cleanupStaleViewers_MaxMinutes=3
 mongoDB_cleanupLogs_MaxDays=14
 mongoDB_cleanupLogs_MaxFinestDays=3
 mongoDB_ResetIndexes=0
